Cleaned and added comments to Top Down Template Blueprint. Updated template description.

diff --git a/Templates/TP_TopDownBP/Config/TemplateDefs.ini b/Templates/TP_TopDownBP/Config/TemplateDefs.ini
--- a/Templates/TP_TopDownBP/Config/TemplateDefs.ini
+++ b/Templates/TP_TopDownBP/Config/TemplateDefs.ini
@@ -4,7 +4,7 @@ Categories=Games
 
 
 LocalizedDisplayNames=(Language="en",Text="Top Down")
-LocalizedDescriptions=(Language="en",Text="The top down template features a playable character where the camera is positioned behind and above the avatar at greater distance than in third person perspective. The player is controlled by using the mouse or touchscreen to click on the required destination and uses the navigation system to facilitate the character movement. This perspective is often used in action role playing games.")
+LocalizedDescriptions=(Language="en",Text="The top down template features a playable character where the camera is positioned behind and above the avatar at greater distance than in third person perspective. The player is controlled by using the mouse, gamepad, or touchscreen and moves the character by clicking on a destination, or following the cursor if input is held. This perspective is often used in action role playing games.")
 
 LocalizedDisplayNames=(Language="ko",Text="내려보기")
 LocalizedDescriptions=(Language="ko",Text="내려보기 템플릿은 플레이어 캐릭터 위 삼인칭 시점보다 멀리 카메라가 위치합니다. 플레이어 제어는 마우스나 터치스크린으로 원하는 목적지를 클릭하는 식으로 이루어지며, 원활한 캐릭터 이동을 위해 내비게이션 시스템을 사용합니다. 이 시점은 액션 롤플레잉 게임에 자주 쓰입니다.")
@@ -15,8 +15,8 @@ LocalizedDescriptions=(Language="ja",Text="トップダウンテンプレート
 LocalizedDisplayNames=(Language="zh-Hans",Text="俯视角游戏")
 LocalizedDescriptions=(Language="zh-Hans",Text="俯视角游戏模板拥有一个可游玩角色，摄像机位于角色的后上方，但位置比第三人称视角要高得多。玩家可以使用鼠标或触摸屏点击目的地，或使用导航系统来辅助角色移动。此视角常用于动作角色扮演类游戏。")
 
-ClassTypes=GameMode, Character, CameraComponent, SpringArmComponent, PlayerController
-AssetTypes=Animation Sequence, Skeletal Mesh, Animation Blueprint,Skeleton, Blend Space 1D
+ClassTypes=Character, PlayerController, CameraComponent, SpringArmComponent, GameMode
+AssetTypes=Animation Sequence, Skeletal Mesh, Animation Blueprint,Skeleton
 
 FoldersToIgnore=Binaries
 FoldersToIgnore=Build
diff --git a/Templates/TP_TopDownBP/manifest.json b/Templates/TP_TopDownBP/manifest.json
--- a/Templates/TP_TopDownBP/manifest.json
+++ b/Templates/TP_TopDownBP/manifest.json
@@ -23,7 +23,7 @@
 	[
 		{
 			"Language": "en",
-			"Text": "The top down pack features a playable character where the camera is positioned behind and above the avatar at greater distance than in third person perspective. The player is controlled by using the mouse or touchscreen to click on the required destination and uses the navigation system to facilitate the character movement. This perspective is often used in action role playing games."
+			"Text": "The top down template features a playable character where the camera is positioned behind and above the avatar at greater distance than in third person perspective. The player is controlled by using the mouse, gamepad, or touchscreen and moves the character by clicking on a destination, or following the cursor if input is held. This perspective is often used in action role playing games."
 		},
 		{
 			"Language": "ko",
@@ -42,17 +42,17 @@
 	[
 		{
 			"Language": "en",
-			"Text": "Animation Sequence, Skeletal Mesh, Animation Blueprint,Skeleton, Blend Space 1D"
+			"Text": "Animation Sequence, Skeletal Mesh, Animation Blueprint, Skeleton"
 		}
 	],	
 	"SearchTags":
 	[
 		{
 			"Language": "en",
-			"Text": "Role, RPG,Role Playing,Top Down"
+			"Text": "RPG,Role Playing,Top Down"
 		}
 	],
-	"ClassTypes": "GameMode, Character, CameraComponent, SpringArmComponent, PlayerController",
+	"ClassTypes": "Character, PlayerController, CameraComponent, SpringArmComponent, GameMode",
 	"Category": "BlueprintFeature",
 	"FocusAsset": "/Game/TopDown/Blueprints",
 	"Thumbnail": "TP_TopDownBP.png",
