New Templates - Top Down
Added content packs for Top Down variants.
Updated base template content for TopDownBP
Added variant pack data to TopDownBP
Fixed a misplaced material on NewVariants TopDown
#rnx

diff --git a/Templates/TP_TopDownBP/Config/DefaultEngine.ini b/Templates/TP_TopDownBP/Config/DefaultEngine.ini
--- a/Templates/TP_TopDownBP/Config/DefaultEngine.ini
+++ b/Templates/TP_TopDownBP/Config/DefaultEngine.ini
@@ -2,8 +2,8 @@
 GameName=TP_TopDownBP
 
 [/Script/EngineSettings.GameMapsSettings]
-EditorStartupMap=/Game/TopDown/Maps/TopDownMap.TopDownMap
-GameDefaultMap=/Game/TopDown/Maps/TopDownMap.TopDownMap
+EditorStartupMap=/Game/TopDown/Lvl_TopDown.Lvl_TopDown
+GameDefaultMap=/Game/TopDown/Lvl_TopDown.Lvl_TopDown
 TransitionMap=
 bUseSplitscreen=True
 TwoPlayerSplitscreenLayout=Horizontal
@@ -90,3 +90,17 @@ SoundCueCookQualityIndex=-1
 -TargetedRHIs=SF_VULKAN_SM5
 +TargetedRHIs=SF_VULKAN_SM6
 
+[/Script/AndroidFileServerEditor.AndroidFileServerRuntimeSettings]
+bEnablePlugin=True
+bAllowNetworkConnection=True
+SecurityToken=7F2D0E674D61A6F3260B0082E348BC25
+bIncludeInShipping=False
+bAllowExternalStartInShipping=False
+bCompileAFSProject=False
+bUseCompression=False
+bLogFiles=False
+bReportStats=False
+ConnectionType=USBOnly
+bUseManualIPAddress=False
+ManualIPAddress=
+
diff --git a/Templates/TP_TopDownBP/Config/TemplateDefs.ini b/Templates/TP_TopDownBP/Config/TemplateDefs.ini
--- a/Templates/TP_TopDownBP/Config/TemplateDefs.ini
+++ b/Templates/TP_TopDownBP/Config/TemplateDefs.ini
@@ -49,3 +49,13 @@ SharedContentPacks=(MountName="LevelPrototyping",DetailLevels=("High"))
 SharedContentPacks=(MountName="Characters",DetailLevels=("High"))
 SharedContentPacks=(MountName="Cursor",DetailLevels=("Standard"))
 EditDetailLevelPreference="High"
+
+; Variant packs
+; NOTE: make sure these are commented out for release!
+;SharedContentPacks=(MountName="Variant_Strategy",DetailLevels=("Standard"))
+;SharedContentPacks=(MountName="Variant_TwinStick",DetailLevels=("Standard"))
+
+; Variant definitions
+Variants=(Name="Strategy",LocalizedDisplayNames=((Language="en",Text="Strategy")),LocalizedDescriptions=((Language="en",Text="A top down strategy game with point and click interactions and selectable units.")),SharedContentPacks=((DetailLevels=(Standard),MountName="Variant_Strategy")))
+
+Variants=(Name="TwinStick",LocalizedDisplayNames=((Language="en",Text="Twin Stick")),LocalizedDescriptions=((Language="en",Text="A twin stick shooter game with infinitely spawning enemies.")),SharedContentPacks=((DetailLevels=(Standard),MountName="Variant_TwinStick")))
\ No newline at end of file
diff --git a/Templates/TP_TopDownBP/TP_TopDownBP.uproject b/Templates/TP_TopDownBP/TP_TopDownBP.uproject
--- a/Templates/TP_TopDownBP/TP_TopDownBP.uproject
+++ b/Templates/TP_TopDownBP/TP_TopDownBP.uproject
@@ -10,6 +10,10 @@
 			"TargetAllowList": [
 				"Editor"
 			]
+		},
+		{
+			"Name": "GameplayStateTree",
+			"Enabled": true
 		}
 	]
 }
\ No newline at end of file
