Templates - TopDownBP
Renamed and reorganized variant BPs
Fixed multiple bugs on variants.
Added Asset Guidelines to variants.
#rnx

diff --git a/Templates/TP_TopDownBP/Config/DefaultEngine.ini b/Templates/TP_TopDownBP/Config/DefaultEngine.ini
--- a/Templates/TP_TopDownBP/Config/DefaultEngine.ini
+++ b/Templates/TP_TopDownBP/Config/DefaultEngine.ini
@@ -48,6 +48,7 @@ TileSetUpdateInterval=1.000000
 HeuristicScale=0.999000
 VerticalDeviationFromGroundCompensation=0.000000
 bForceRebuildOnLoad=True
+RuntimeGeneration=Dynamic
 
 [/Script/Engine.RendererSettings]
 r.ReflectionMethod=1
