Templates - Top Down - setup template def to include Input content pak for Top Down Twin Stick only - create Touch Aim input function to handle touch controls more cleanly

diff --git a/Templates/TP_TopDownBP/Config/TemplateDefs.ini b/Templates/TP_TopDownBP/Config/TemplateDefs.ini
--- a/Templates/TP_TopDownBP/Config/TemplateDefs.ini
+++ b/Templates/TP_TopDownBP/Config/TemplateDefs.ini
@@ -54,8 +54,8 @@ EditDetailLevelPreference="High"
 ; NOTE: make sure these are commented out for release!
 ;SharedContentPacks=(MountName="Variant_Strategy",DetailLevels=("Standard"))
 ;SharedContentPacks=(MountName="Variant_TwinStick",DetailLevels=("Standard"))
+;SharedContentPacks=(MountName="Input",DetailLevels=("High"))
 
 ; Variant definitions
 Variants=(Name="Strategy",LocalizedDisplayNames=((Language="en",Text="Strategy")),LocalizedDescriptions=((Language="en",Text="A top down strategy game with point and click interactions and selectable units.")),SharedContentPacks=((DetailLevels=(Standard),MountName="Variant_Strategy")))
-
-Variants=(Name="TwinStick",LocalizedDisplayNames=((Language="en",Text="Twin Stick")),LocalizedDescriptions=((Language="en",Text="A twin stick shooter game with infinitely spawning enemies.")),SharedContentPacks=((DetailLevels=(Standard),MountName="Variant_TwinStick")))
\ No newline at end of file
+Variants=(Name="TwinStick",LocalizedDisplayNames=((Language="en",Text="Twin Stick")),LocalizedDescriptions=((Language="en",Text="A twin stick shooter game with infinitely spawning enemies.")),SharedContentPacks=((DetailLevels=(Standard),MountName="Variant_TwinStick"),(MountName="Input",DetailLevels=("High")))
\ No newline at end of file
